[NO NON-PII COMMIT MESSAGE]

diff --git a/Engine/Source/Runtime/AudioCaptureCore/Private/AudioCaptureInternal.h b/Engine/Source/Runtime/AudioCaptureCore/Private/AudioCaptureInternal.h
--- a/Engine/Source/Runtime/AudioCaptureCore/Private/AudioCaptureInternal.h
+++ b/Engine/Source/Runtime/AudioCaptureCore/Private/AudioCaptureInternal.h
@@ -40,7 +40,7 @@ namespace Audio
 		}
 		else
 		{
-			UE_LOG(LogAudioCaptureCore, Warning, TEXT("No Audio Capture implementations found. Audio input will be silent."));
+			UE_LOG(LogAudioCaptureCore, Display, TEXT("No Audio Capture implementations found. Audio input will be silent."));
 			return TUniquePtr<IAudioCaptureStream>(new FNullAudioCaptureDevice());
 		}
 	}
