FAudioCapture factory now respects GEngine->UseSound() so that capture devices are only instantiated when audio is active.

diff --git a/Engine/Source/Runtime/AudioCaptureCore/AudioCaptureCore.Build.cs b/Engine/Source/Runtime/AudioCaptureCore/AudioCaptureCore.Build.cs
--- a/Engine/Source/Runtime/AudioCaptureCore/AudioCaptureCore.Build.cs
+++ b/Engine/Source/Runtime/AudioCaptureCore/AudioCaptureCore.Build.cs
@@ -13,6 +13,13 @@ namespace UnrealBuildTool.Rules
                     "SignalProcessing"
                 }
             );
-        }
+
+			PrivateDependencyModuleNames.AddRange(
+				new string[]
+				{
+					"Engine"
+				}
+			);
+		}
 	}
 }
diff --git a/Engine/Source/Runtime/AudioCaptureCore/Private/AudioCaptureInternal.h b/Engine/Source/Runtime/AudioCaptureCore/Private/AudioCaptureInternal.h
--- a/Engine/Source/Runtime/AudioCaptureCore/Private/AudioCaptureInternal.h
+++ b/Engine/Source/Runtime/AudioCaptureCore/Private/AudioCaptureInternal.h
@@ -4,6 +4,7 @@
 
 #include "AudioCaptureCore.h"
 #include "AudioCaptureCoreLog.h"
+#include "Engine/Engine.h"
 #include "Features/IModularFeatures.h"
 
 namespace Audio
@@ -37,7 +38,7 @@ namespace Audio
 		IModularFeatures::Get().UnlockModularFeatureList();
 
 		// For now, just return the first audio capture stream implemented. We can make this configurable at a later point.
-		if (AudioCaptureStreamFactories.Num() > 0 && AudioCaptureStreamFactories[0] != nullptr)
+		if (AudioCaptureStreamFactories.Num() > 0 && AudioCaptureStreamFactories[0] != nullptr && GEngine->UseSound())
 		{
 			return AudioCaptureStreamFactories[0]->CreateNewAudioCaptureStream();
 		}
