[Metasounds] Audio Capture/Timecode Provider initialization order incorrect

#rnx

diff --git a/Engine/Source/Runtime/AudioCaptureCore/Private/AudioCaptureInternal.h b/Engine/Source/Runtime/AudioCaptureCore/Private/AudioCaptureInternal.h
--- a/Engine/Source/Runtime/AudioCaptureCore/Private/AudioCaptureInternal.h
+++ b/Engine/Source/Runtime/AudioCaptureCore/Private/AudioCaptureInternal.h
@@ -4,7 +4,7 @@
 
 #include "AudioCaptureCore.h"
 #include "AudioCaptureCoreLog.h"
-#include "Engine/Engine.h"
+#include "Misc/App.h"
 #include "Features/IModularFeatures.h"
 
 namespace Audio
@@ -38,7 +38,7 @@ namespace Audio
 		IModularFeatures::Get().UnlockModularFeatureList();
 
 		// For now, just return the first audio capture stream implemented. We can make this configurable at a later point.
-		if (AudioCaptureStreamFactories.Num() > 0 && AudioCaptureStreamFactories[0] != nullptr && GEngine->UseSound())
+		if (AudioCaptureStreamFactories.Num() > 0 && AudioCaptureStreamFactories[0] != nullptr && FApp::CanEverRenderAudio())
 		{
 			return AudioCaptureStreamFactories[0]->CreateNewAudioCaptureStream();
 		}
