[AudioCapture] Fix crash when WASAPI capture thread fails to start

FAudioCaptureSynth::OpenDefaultStream() ignored the return value of StartStream(). If the capture thread failed to start (e.g. intermittent OS resource pressure), the stream entered INVALID_STATE but the caller proceeded as if capture was active. This caused an assert in StopCapturing() in dev builds, and a null COM dereference in shipping builds where the default destructor tore down members without orderly shutdown.

1. Check StartStream() return value in OpenDefaultStream; close the stream and report failure on error.
2. Call CloseStream() in ~FWasapiStreamManager to guarantee orderly teardown regardless of how destruction is reached.

#Synced-CL 52421152

diff --git a/Engine/Source/Runtime/AudioCaptureCore/Private/AudioCaptureCore.cpp b/Engine/Source/Runtime/AudioCaptureCore/Private/AudioCaptureCore.cpp
--- a/Engine/Source/Runtime/AudioCaptureCore/Private/AudioCaptureCore.cpp
+++ b/Engine/Source/Runtime/AudioCaptureCore/Private/AudioCaptureCore.cpp
@@ -223,10 +223,15 @@ namespace Audio
 
 			FAudioCaptureDeviceParams Params = FAudioCaptureDeviceParams();
 
-			// Start the stream here to avoid hitching the audio render thread. 
+			// Start the stream here to avoid hitching the audio render thread.
 			if (AudioCapture.OpenAudioCaptureStream(Params, MoveTemp(OnCapture), 1024))
 			{
-				AudioCapture.StartStream();
+				if (!AudioCapture.StartStream())
+				{
+					UE_LOG(LogAudioCaptureCore, Error, TEXT("FAudioCaptureSynth::OpenDefaultStream: Failed to start audio capture stream."));
+					AudioCapture.CloseStream();
+					bSuccess = false;
+				}
 			}
 			else
 			{
diff --git a/Engine/Source/Runtime/AudioCaptureImplementations/Windows/AudioCaputureWasapi/Private/WasapiStreamManager.cpp b/Engine/Source/Runtime/AudioCaptureImplementations/Windows/AudioCaputureWasapi/Private/WasapiStreamManager.cpp
--- a/Engine/Source/Runtime/AudioCaptureImplementations/Windows/AudioCaputureWasapi/Private/WasapiStreamManager.cpp
+++ b/Engine/Source/Runtime/AudioCaptureImplementations/Windows/AudioCaputureWasapi/Private/WasapiStreamManager.cpp
@@ -13,6 +13,11 @@ namespace Audio
 		DeviceEnumerator.Initialize();
 	}
 
+	FWasapiStreamManager::~FWasapiStreamManager()
+	{
+		CloseStream();
+	}
+
 	FString FWasapiStreamManager::GetDefaultInputDeviceId()
 	{
 		return DeviceEnumerator.GetDefaultInputDeviceId();
diff --git a/Engine/Source/Runtime/AudioCaptureImplementations/Windows/AudioCaputureWasapi/Private/WasapiStreamManager.h b/Engine/Source/Runtime/AudioCaptureImplementations/Windows/AudioCaputureWasapi/Private/WasapiStreamManager.h
--- a/Engine/Source/Runtime/AudioCaptureImplementations/Windows/AudioCaputureWasapi/Private/WasapiStreamManager.h
+++ b/Engine/Source/Runtime/AudioCaptureImplementations/Windows/AudioCaputureWasapi/Private/WasapiStreamManager.h
@@ -25,7 +25,7 @@ namespace Audio
 		FWasapiStreamManager(FWasapiStreamManager&& InOther) = delete;
 		FWasapiStreamManager(const FWasapiStreamManager& InOther) = delete;
 
-		virtual ~FWasapiStreamManager() = default;
+		virtual ~FWasapiStreamManager();
 
 		FWasapiStreamManager& operator=(FWasapiStreamManager&& InOther) = delete;
 		FWasapiStreamManager& operator=(const FWasapiStreamManager& InOther) = delete;
